Take Recorder - Move SerializedRecorderInterface into Runtime.

The SerializedRecorderInterface is used within the Take Recorder animation recording logic as well as other parts of the Take Recorder modules. To allow us to move these TakeRecorder modules to be runtime, we also need to move this module from Editor to Runtime. This is a clean move with no issues of moving of names, renaming etc. This is simply a move with a small change to guard against one use of an editor function for Property->GetDisplayName() which is editor only.

#rnx
#Synced-CL 52673305

diff --git a/Engine/Source/Editor/SerializedRecorderInterface/Private/MovieSceneSectionSerialization.cpp b/Engine/Source/Editor/SerializedRecorderInterface/Private/MovieSceneSectionSerialization.cpp
deleted file mode 100644
diff --git a/Engine/Source/Editor/SerializedRecorderInterface/Private/SerializedRecorderInterfaceModule.cpp b/Engine/Source/Editor/SerializedRecorderInterface/Private/SerializedRecorderInterfaceModule.cpp
deleted file mode 100644
diff --git a/Engine/Source/Editor/SerializedRecorderInterface/Public/ISerializedRecorder.h b/Engine/Source/Editor/SerializedRecorderInterface/Public/ISerializedRecorder.h
deleted file mode 100644
diff --git a/Engine/Source/Editor/SerializedRecorderInterface/Public/ISerializedRecorderInterfaceModule.h b/Engine/Source/Editor/SerializedRecorderInterface/Public/ISerializedRecorderInterfaceModule.h
deleted file mode 100644
diff --git a/Engine/Source/Editor/SerializedRecorderInterface/Public/Serializers/MovieSceneActorSerialization.h b/Engine/Source/Editor/SerializedRecorderInterface/Public/Serializers/MovieSceneActorSerialization.h
deleted file mode 100644
diff --git a/Engine/Source/Editor/SerializedRecorderInterface/Public/Serializers/MovieSceneAnimationSerialization.h b/Engine/Source/Editor/SerializedRecorderInterface/Public/Serializers/MovieSceneAnimationSerialization.h
deleted file mode 100644
diff --git a/Engine/Source/Editor/SerializedRecorderInterface/Public/Serializers/MovieSceneLiveLinkSerialization.h b/Engine/Source/Editor/SerializedRecorderInterface/Public/Serializers/MovieSceneLiveLinkSerialization.h
deleted file mode 100644
diff --git a/Engine/Source/Editor/SerializedRecorderInterface/Public/Serializers/MovieSceneManifestSerialization.h b/Engine/Source/Editor/SerializedRecorderInterface/Public/Serializers/MovieSceneManifestSerialization.h
deleted file mode 100644
diff --git a/Engine/Source/Editor/SerializedRecorderInterface/Public/Serializers/MovieScenePropertySerialization.h b/Engine/Source/Editor/SerializedRecorderInterface/Public/Serializers/MovieScenePropertySerialization.h
deleted file mode 100644
diff --git a/Engine/Source/Editor/SerializedRecorderInterface/Public/Serializers/MovieSceneSectionSerialization.h b/Engine/Source/Editor/SerializedRecorderInterface/Public/Serializers/MovieSceneSectionSerialization.h
deleted file mode 100644
diff --git a/Engine/Source/Editor/SerializedRecorderInterface/Public/Serializers/MovieSceneSerializedType.h b/Engine/Source/Editor/SerializedRecorderInterface/Public/Serializers/MovieSceneSerializedType.h
deleted file mode 100644
diff --git a/Engine/Source/Editor/SerializedRecorderInterface/Public/Serializers/MovieSceneSpawnSerialization.h b/Engine/Source/Editor/SerializedRecorderInterface/Public/Serializers/MovieSceneSpawnSerialization.h
deleted file mode 100644
diff --git a/Engine/Source/Editor/SerializedRecorderInterface/Public/Serializers/MovieSceneTransformSerialization.h b/Engine/Source/Editor/SerializedRecorderInterface/Public/Serializers/MovieSceneTransformSerialization.h
deleted file mode 100644
diff --git a/Engine/Source/Editor/SerializedRecorderInterface/SerializedRecorderInterface.Build.cs b/Engine/Source/Editor/SerializedRecorderInterface/SerializedRecorderInterface.Build.cs
deleted file mode 100644
diff --git a/Engine/Source/Runtime/SerializedRecorderInterface/Private/MovieSceneSectionSerialization.cpp b/Engine/Source/Runtime/SerializedRecorderInterface/Private/MovieSceneSectionSerialization.cpp
new file mode 100644
--- /dev/null
+++ b/Engine/Source/Runtime/SerializedRecorderInterface/Private/MovieSceneSectionSerialization.cpp
@@ -0,0 +1,130 @@
+// Copyright Epic Games, Inc. All Rights Reserved.
+
+#include "Serializers/MovieSceneSectionSerialization.h"
+
+#include "Math/MathFwd.h"
+#include "Serialization/CustomVersion.h"
+
+struct FActorFileHeader;
+struct FActorProperty;
+struct FAnimationFileHeader;
+struct FColor;
+struct FLinearColor;
+struct FManifestFileHeader;
+struct FManifestProperty;
+struct FPropertyFileHeader;
+struct FSerializedAnimation;
+struct FSerializedTransform;
+struct FSerializedTypeFileHeader;
+struct FSpawnFileHeader;
+struct FSpawnProperty;
+struct FTransformFileHeader;
+template <typename PropertyType> struct FSerializedProperty;
+
+DEFINE_LOG_CATEGORY(MovieSceneSerialization);
+
+
+const FGuid FTempCustomVersion::GUID(0xCB8AB0CD, 0xE78C4BDE, 0xA8621393, 0x14E9EF62);
+FCustomVersionRegistration GRegisterLiveLinkMovieSectionCustomVersion(FTempCustomVersion::GUID, FTempCustomVersion::LatestVersion, TEXT("LiveLinkCustomVersion"));
+
+const int64 MovieSceneSerializationNamespace::InvalidOffset = -1;
+const float MovieSceneSerializationNamespace::SerializerSleepTime = 0.2f;
+bool MovieSceneSerializationNamespace::bAutoSerialize = false;
+
+/** ==  Template Statics Defined  --  **/
+template<> SERIALIZEDRECORDERINTERFACE_API
+FMovieSceneSerializerRunnable<FPropertyFileHeader, FPropertyFileHeader>* TMovieSceneSerializer<FPropertyFileHeader, FPropertyFileHeader>::Runnable = nullptr;
+template<> SERIALIZEDRECORDERINTERFACE_API
+FRunnableThread* TMovieSceneSerializer<FPropertyFileHeader, FPropertyFileHeader>::Thread = nullptr;
+
+template<> SERIALIZEDRECORDERINTERFACE_API
+FMovieSceneSerializerRunnable<FTransformFileHeader, FSerializedTransform>* TMovieSceneSerializer<FTransformFileHeader, FSerializedTransform>::Runnable = nullptr;
+template<> SERIALIZEDRECORDERINTERFACE_API
+FRunnableThread* TMovieSceneSerializer<FTransformFileHeader, FSerializedTransform>::Thread = nullptr;
+
+template<> SERIALIZEDRECORDERINTERFACE_API
+FMovieSceneSerializerRunnable<FAnimationFileHeader, FSerializedAnimation>* TMovieSceneSerializer<FAnimationFileHeader, FSerializedAnimation>::Runnable = nullptr;
+template<> SERIALIZEDRECORDERINTERFACE_API
+FRunnableThread* TMovieSceneSerializer<FAnimationFileHeader, FSerializedAnimation>::Thread = nullptr;
+
+
+template<> SERIALIZEDRECORDERINTERFACE_API
+FMovieSceneSerializerRunnable<FSpawnFileHeader, FSpawnProperty>* TMovieSceneSerializer<FSpawnFileHeader, FSpawnProperty>::Runnable = nullptr;
+template<> SERIALIZEDRECORDERINTERFACE_API
+FRunnableThread* TMovieSceneSerializer<FSpawnFileHeader, FSpawnProperty>::Thread = nullptr;
+
+template<> SERIALIZEDRECORDERINTERFACE_API
+FMovieSceneSerializerRunnable<FPropertyFileHeader, FSerializedProperty<int64>>* TMovieSceneSerializer<FPropertyFileHeader, FSerializedProperty<int64>>::Runnable = nullptr;
+template<> SERIALIZEDRECORDERINTERFACE_API
+FRunnableThread* TMovieSceneSerializer<FPropertyFileHeader, FSerializedProperty<int64>>::Thread = nullptr;
+
+template<> SERIALIZEDRECORDERINTERFACE_API
+FMovieSceneSerializerRunnable<FPropertyFileHeader, FSerializedProperty<bool>>* TMovieSceneSerializer<FPropertyFileHeader, FSerializedProperty<bool>>::Runnable = nullptr;
+template<> SERIALIZEDRECORDERINTERFACE_API
+FRunnableThread* TMovieSceneSerializer<FPropertyFileHeader, FSerializedProperty<bool>>::Thread = nullptr;
+
+template<> SERIALIZEDRECORDERINTERFACE_API
+FMovieSceneSerializerRunnable<FPropertyFileHeader, FSerializedProperty<uint8>>* TMovieSceneSerializer<FPropertyFileHeader, FSerializedProperty<uint8>>::Runnable = nullptr;
+template<> SERIALIZEDRECORDERINTERFACE_API
+FRunnableThread* TMovieSceneSerializer<FPropertyFileHeader, FSerializedProperty<uint8>>::Thread = nullptr;
+
+template<> SERIALIZEDRECORDERINTERFACE_API
+FMovieSceneSerializerRunnable<FPropertyFileHeader, FSerializedProperty<float>>* TMovieSceneSerializer<FPropertyFileHeader, FSerializedProperty<float>>::Runnable = nullptr;
+template<> SERIALIZEDRECORDERINTERFACE_API
+FRunnableThread* TMovieSceneSerializer<FPropertyFileHeader, FSerializedProperty<float>>::Thread = nullptr;
+
+template<> SERIALIZEDRECORDERINTERFACE_API
+FMovieSceneSerializerRunnable<FPropertyFileHeader, FSerializedProperty<double>>* TMovieSceneSerializer<FPropertyFileHeader, FSerializedProperty<double>>::Runnable = nullptr;
+template<> SERIALIZEDRECORDERINTERFACE_API
+FRunnableThread* TMovieSceneSerializer<FPropertyFileHeader, FSerializedProperty<double>>::Thread = nullptr;
+
+template<> SERIALIZEDRECORDERINTERFACE_API
+FMovieSceneSerializerRunnable<FPropertyFileHeader, FSerializedProperty<FVector3f>>* TMovieSceneSerializer<FPropertyFileHeader, FSerializedProperty<FVector3f>>::Runnable = nullptr;
+template<> SERIALIZEDRECORDERINTERFACE_API
+FRunnableThread* TMovieSceneSerializer<FPropertyFileHeader, FSerializedProperty<FVector3f>>::Thread = nullptr;
+
+template<> SERIALIZEDRECORDERINTERFACE_API
+FMovieSceneSerializerRunnable<FPropertyFileHeader, FSerializedProperty<FVector3d>>* TMovieSceneSerializer<FPropertyFileHeader, FSerializedProperty<FVector3d>>::Runnable = nullptr;
+template<> SERIALIZEDRECORDERINTERFACE_API
+FRunnableThread* TMovieSceneSerializer<FPropertyFileHeader, FSerializedProperty<FVector3d>>::Thread = nullptr;
+
+template<> SERIALIZEDRECORDERINTERFACE_API
+FMovieSceneSerializerRunnable<FPropertyFileHeader, FSerializedProperty<FColor>>* TMovieSceneSerializer<FPropertyFileHeader, FSerializedProperty<FColor>>::Runnable = nullptr;
+template<> SERIALIZEDRECORDERINTERFACE_API
+FRunnableThread* TMovieSceneSerializer<FPropertyFileHeader, FSerializedProperty<FColor>>::Thread = nullptr;
+
+template<> SERIALIZEDRECORDERINTERFACE_API
+FMovieSceneSerializerRunnable<FPropertyFileHeader, FSerializedProperty<int32>>* TMovieSceneSerializer<FPropertyFileHeader, FSerializedProperty<int32>>::Runnable = nullptr;
+template<> SERIALIZEDRECORDERINTERFACE_API
+FRunnableThread* TMovieSceneSerializer<FPropertyFileHeader, FSerializedProperty<int32>>::Thread = nullptr;
+
+template<> SERIALIZEDRECORDERINTERFACE_API
+FMovieSceneSerializerRunnable<FPropertyFileHeader, FSerializedProperty<FString>>* TMovieSceneSerializer<FPropertyFileHeader, FSerializedProperty<FString>>::Runnable = nullptr;
+template<> SERIALIZEDRECORDERINTERFACE_API
+FRunnableThread* TMovieSceneSerializer<FPropertyFileHeader, FSerializedProperty<FString>>::Thread = nullptr;
+
+template<> SERIALIZEDRECORDERINTERFACE_API
+FMovieSceneSerializerRunnable<FPropertyFileHeader, FSerializedProperty<FLinearColor>>* TMovieSceneSerializer<FPropertyFileHeader, FSerializedProperty<FLinearColor>>::Runnable = nullptr;
+template<> SERIALIZEDRECORDERINTERFACE_API
+FRunnableThread* TMovieSceneSerializer<FPropertyFileHeader, FSerializedProperty<FLinearColor>>::Thread = nullptr;
+
+/** External Tempplate statics.  Need to define here due to linkage issues*/
+
+/* From TakeRecorderActorSource */
+template<> SERIALIZEDRECORDERINTERFACE_API
+FMovieSceneSerializerRunnable<FManifestFileHeader, FManifestProperty>* TMovieSceneSerializer<FManifestFileHeader, FManifestProperty>::Runnable = nullptr;
+template<> SERIALIZEDRECORDERINTERFACE_API
+FRunnableThread* TMovieSceneSerializer<FManifestFileHeader, FManifestProperty>::Thread = nullptr;
+
+template<> SERIALIZEDRECORDERINTERFACE_API
+FMovieSceneSerializerRunnable<FActorFileHeader, FActorProperty>* TMovieSceneSerializer<FActorFileHeader, FActorProperty>::Runnable = nullptr;
+template<> SERIALIZEDRECORDERINTERFACE_API
+FRunnableThread* TMovieSceneSerializer<FActorFileHeader, FActorProperty>::Thread = nullptr;
+
+
+/* From SerializedRecorder */
+template<> SERIALIZEDRECORDERINTERFACE_API
+FMovieSceneSerializerRunnable<FSerializedTypeFileHeader, FSerializedTypeFileHeader>* TMovieSceneSerializer<FSerializedTypeFileHeader, FSerializedTypeFileHeader>::Runnable = nullptr;
+template<> SERIALIZEDRECORDERINTERFACE_API
+FRunnableThread* TMovieSceneSerializer<FSerializedTypeFileHeader, FSerializedTypeFileHeader>::Thread = nullptr;
+
diff --git a/Engine/Source/Runtime/SerializedRecorderInterface/Private/SerializedRecorderInterfaceModule.cpp b/Engine/Source/Runtime/SerializedRecorderInterface/Private/SerializedRecorderInterfaceModule.cpp
new file mode 100644
--- /dev/null
+++ b/Engine/Source/Runtime/SerializedRecorderInterface/Private/SerializedRecorderInterfaceModule.cpp
@@ -0,0 +1,30 @@
+// Copyright Epic Games, Inc. All Rights Reserved.
+
+#include "HAL/Platform.h"
+#include "ISerializedRecorder.h"
+#include "ISerializedRecorderInterfaceModule.h"
+#include "Modules/ModuleManager.h"
+#include "UObject/NameTypes.h"
+
+FName ISerializedRecorder::ModularFeatureName(TEXT("ModularFeature_SerialzedRecorder"));
+
+class FSerializedRecorderInterfaceModule : public ISerializedRecorderInterfaceModule
+{
+public:
+	
+	// IModuleInterface interface
+	virtual void StartupModule() override
+	{
+		
+	}
+
+	virtual bool SupportsDynamicReloading() override
+	{
+		return false;
+	}
+
+private:
+
+};
+
+IMPLEMENT_MODULE(FSerializedRecorderInterfaceModule, SerializedRecorderInterface);
diff --git a/Engine/Source/Runtime/SerializedRecorderInterface/Public/ISerializedRecorder.h b/Engine/Source/Runtime/SerializedRecorderInterface/Public/ISerializedRecorder.h
new file mode 100644
--- /dev/null
+++ b/Engine/Source/Runtime/SerializedRecorderInterface/Public/ISerializedRecorder.h
@@ -0,0 +1,22 @@
+// Copyright Epic Games, Inc. All Rights Reserved.
+
+#pragma once
+#include "CoreMinimal.h"
+#include "Features/IModularFeature.h"
+
+#define UE_API SERIALIZEDRECORDERINTERFACE_API
+
+
+class UMovieSceneSequence;
+class UWorld;
+
+class  ISerializedRecorder : public IModularFeature
+{
+public:
+	virtual ~ISerializedRecorder() = default;
+
+	static UE_API FName ModularFeatureName;
+	virtual bool LoadRecordedSequencerFile(UMovieSceneSequence* InMovieSceneSequence, UWorld *PlaybackContext, const FString& InFileName, TFunction<void()> InCompletionCallback) = 0;
+};
+
+#undef UE_API
diff --git a/Engine/Source/Runtime/SerializedRecorderInterface/Public/ISerializedRecorderInterfaceModule.h b/Engine/Source/Runtime/SerializedRecorderInterface/Public/ISerializedRecorderInterfaceModule.h
new file mode 100644
--- /dev/null
+++ b/Engine/Source/Runtime/SerializedRecorderInterface/Public/ISerializedRecorderInterfaceModule.h
@@ -0,0 +1,29 @@
+// Copyright Epic Games, Inc. All Rights Reserved.
+
+#pragma once
+
+#include "Modules/ModuleInterface.h"
+#include "Modules/ModuleManager.h"
+
+
+class ISerializedRecorderInterfaceModule
+	: public IModuleInterface
+{
+public:
+
+	static ISerializedRecorderInterfaceModule& Get()
+	{
+#if PLATFORM_IOS
+        static ISerializedRecorderInterfaceModule& LiveLinkInterfaceModule = FModuleManager::LoadModuleChecked<ISerializedRecorderInterfaceModule>("LiveLinkInterface");
+        return LiveLinkInterfaceModule;
+#else
+        return FModuleManager::LoadModuleChecked<ISerializedRecorderInterfaceModule>("Interface");
+#endif
+	}
+
+public:
+
+	/** Virtual destructor. */
+	virtual ~ISerializedRecorderInterfaceModule() = default;
+};
+
diff --git a/Engine/Source/Runtime/SerializedRecorderInterface/Public/Serializers/MovieSceneActorSerialization.h b/Engine/Source/Runtime/SerializedRecorderInterface/Public/Serializers/MovieSceneActorSerialization.h
new file mode 100644
--- /dev/null
+++ b/Engine/Source/Runtime/SerializedRecorderInterface/Public/Serializers/MovieSceneActorSerialization.h
@@ -0,0 +1,113 @@
+// Copyright Epic Games, Inc. All Rights Reserved.
+
+#pragma once
+
+#include "Misc/FrameRate.h"
+#include "UObject/UnrealType.h"
+#include "MovieSceneSectionSerialization.h"
+
+
+struct FActorFileHeader
+{
+	static const int32 cVersion = 1;
+
+	FActorFileHeader() : Version(cVersion)
+	{
+	}
+
+
+	FActorFileHeader(const FString& InName, const FString& InLabel, const FName& InSerializedType, const FString& InClassName, bool SpawnedPost)
+		: Version(cVersion)
+		, SerializedType(InSerializedType)
+		, UObjectName(InName)
+		, Label(InLabel)
+		, bRecordToPossessable(false)
+		, bWasSpawnedPostRecord(SpawnedPost)
+		, ClassName(InClassName)
+		, TemplateName("None")
+	{
+	}
+
+	friend FArchive& operator<<(FArchive& Ar, FActorFileHeader& Header)
+	{
+		Ar << Header.Version;
+		Ar << Header.SerializedType;
+		Ar << Header.Guid;
+		Ar << Header.UObjectName;
+		Ar << Header.Label;
+		Ar << Header.bRecordToPossessable;
+		Ar << Header.bWasSpawnedPostRecord;
+		Ar << Header.ClassName;
+		Ar << Header.TemplateName;
+		Ar << Header.FolderName;
+
+		return Ar;
+	}
+
+
+	//DATA
+	int32 Version;
+	FName SerializedType;
+	FGuid Guid;
+
+	FString UObjectName;
+	FString Label;
+	bool bRecordToPossessable;
+	bool bWasSpawnedPostRecord;
+	FString ClassName;
+	FString TemplateName;
+	FName FolderName;
+};
+
+enum class EActoryPropertyType : uint8 {
+
+	ComponentType,
+	PropertyType,
+	OtherType,
+};
+
+struct FActorProperty
+{
+	FActorProperty() = default;
+	FActorProperty(const FString &InObjectName, const FName& InSerializedType, const FGuid& InGuid) :
+		UObjectName(InObjectName)
+		, SerializedType(InSerializedType)
+		, Guid(InGuid)
+		, Type(EActoryPropertyType::OtherType)
+	{}
+	friend FArchive& operator<<(FArchive& Ar, FActorProperty& Property)
+	{
+
+		Ar << Property.SerializedType;
+		Ar << Property.Guid;
+		Ar << Property.UObjectName;
+		Ar << Property.Type;
+		if (Property.Type == EActoryPropertyType::ComponentType)
+		{
+			Ar << Property.BindingName;
+			Ar << Property.ClassName;
+		}
+		else if (Property.Type == EActoryPropertyType::PropertyType)
+		{
+			Ar << Property.PropertyName;
+		}
+
+		return Ar;
+	}
+
+
+	//DATA
+
+	FString UObjectName;
+	FName SerializedType;
+	FGuid Guid;
+	EActoryPropertyType  Type;
+	FString PropertyName;
+	FString BindingName;
+	FString ClassName;
+
+};
+
+
+using FActorSerializedFrame = TMovieSceneSerializedFrame<FActorProperty>;
+using FActorSerializer = TMovieSceneSerializer<FActorFileHeader, FActorProperty>;
diff --git a/Engine/Source/Runtime/SerializedRecorderInterface/Public/Serializers/MovieSceneAnimationSerialization.h b/Engine/Source/Runtime/SerializedRecorderInterface/Public/Serializers/MovieSceneAnimationSerialization.h
new file mode 100644
--- /dev/null
+++ b/Engine/Source/Runtime/SerializedRecorderInterface/Public/Serializers/MovieSceneAnimationSerialization.h
@@ -0,0 +1,98 @@
+// Copyright Epic Games, Inc. All Rights Reserved.
+
+#pragma once
+
+#include "Misc/FrameRate.h"
+#include "UObject/UnrealType.h"
+#include "MovieSceneSectionSerialization.h"
+
+
+struct FAnimationFileHeader
+{
+	static const int32 cVersion = 1;
+
+	FAnimationFileHeader() : Version(cVersion)
+	{
+	}
+
+
+	FAnimationFileHeader(const FName& InSerializedType, FGuid& InGuid, float InIntervalTime)
+		: Version(cVersion)
+		, SerializedType(InSerializedType) 
+		, Guid(InGuid)
+		, IntervalTime(InIntervalTime)
+	{};
+
+
+	friend FArchive& operator<<(FArchive& Ar, FAnimationFileHeader& Header)
+	{
+		Ar << Header.Version;
+		Ar << Header.SerializedType;
+		Ar << Header.Guid;
+		Ar << Header.ActorGuid;
+		Ar << Header.IntervalTime;
+		Ar << Header.StartTime;
+		Ar << Header.AnimationTrackNames;
+		return Ar;
+	}
+
+	void AddNewRawTrack(const FName& BoneTreeName)
+	{
+		AnimationTrackNames.Add(BoneTreeName);
+	}
+
+	//DATA
+	int32 Version;
+	FName SerializedType;
+	FGuid Guid;
+	FGuid ActorGuid;
+	float IntervalTime;
+	float StartTime;
+	TArray<FName> AnimationTrackNames;
+};
+
+struct FSerializedAnimationPerFrame
+{
+	FSerializedAnimationPerFrame() = default;
+
+	friend FArchive& operator<<(FArchive& Ar, FSerializedAnimationPerFrame& Animation)
+	{
+		Ar << Animation.BoneIndex;
+		Ar << Animation.PosKey;
+		Ar << Animation.RotKey;
+		Ar << Animation.ScaleKey;
+		return Ar;
+	}
+
+	int32 BoneIndex;
+	FVector PosKey;
+	FQuat RotKey;
+	FVector ScaleKey;
+
+};
+
+struct FSerializedAnimation
+{
+	FSerializedAnimation() = default;
+
+	void AddTransform(int32 BoneIndex, const FTransform& InTransform)
+	{
+		FSerializedAnimationPerFrame Frame;
+		Frame.BoneIndex = BoneIndex;
+		Frame.PosKey = InTransform.GetTranslation();
+		Frame.RotKey = InTransform.GetRotation();
+		Frame.ScaleKey = InTransform.GetScale3D();
+		AnimationData.Emplace(Frame);
+	}
+
+	friend FArchive& operator<<(FArchive& Ar, FSerializedAnimation& Animation)
+	{
+		Ar << Animation.AnimationData;
+		return Ar;
+	}
+	TArray<FSerializedAnimationPerFrame> AnimationData;
+
+};
+
+using FAnimationSerializedFrame = TMovieSceneSerializedFrame<FSerializedAnimation>;
+using FAnimationSerializer = TMovieSceneSerializer<FAnimationFileHeader, FSerializedAnimation>;
diff --git a/Engine/Source/Runtime/SerializedRecorderInterface/Public/Serializers/MovieSceneLiveLinkSerialization.h b/Engine/Source/Runtime/SerializedRecorderInterface/Public/Serializers/MovieSceneLiveLinkSerialization.h
new file mode 100644
--- /dev/null
+++ b/Engine/Source/Runtime/SerializedRecorderInterface/Public/Serializers/MovieSceneLiveLinkSerialization.h
@@ -0,0 +1,98 @@
+// Copyright Epic Games, Inc. All Rights Reserved.
+
+#pragma once
+
+#include "Misc/FrameRate.h"
+
+#include "LiveLinkRole.h"
+#include "LiveLinkTypes.h"
+#include "LiveLinkCustomVersion.h"
+
+#include "Serializers/MovieSceneSectionSerialization.h"
+
+
+/** !!! LiveLink Track serialization is not supported !!! */
+
+struct FLiveLinkManifestHeader
+{
+
+	static const int32 cVersion = 2;
+
+	FLiveLinkManifestHeader() : Version(cVersion)
+	{
+	}
+
+
+	FLiveLinkManifestHeader(const FName& InSerializedType)
+		: Version(cVersion)
+		, SerializedType(InSerializedType)
+		, bIsManifest(true)
+	{
+	}
+
+	friend FArchive& operator<<(FArchive& Ar, FLiveLinkManifestHeader& Header)
+	{
+		Ar << Header.Version;
+		Ar << Header.SerializedType;
+		Ar << Header.bIsManifest;
+		Ar << Header.SubjectNames;
+
+		return Ar;
+	}
+	int32 Version;
+	FName SerializedType;
+	bool bIsManifest;
+	TArray<FName> SubjectNames;
+};
+
+using FLiveLinkManifestSerializer = TMovieSceneSerializer<FLiveLinkManifestHeader, FLiveLinkManifestHeader>;
+
+
+PRAGMA_DISABLE_DEPRECATION_WARNINGS
+//Header and data for individual subjects that contain the data.
+struct FLiveLinkFileHeader
+{
+	FLiveLinkFileHeader() = default;
+
+	FLiveLinkFileHeader(const FName& InSubjectName, double InSecondsDiff, TSubclassOf<ULiveLinkRole> InSubjectRole, const TSharedPtr<FLiveLinkStaticDataStruct>& InStaticData, const FName& InSerializedType, const FGuid& InGuid)
+		: bIsManifest(false)
+		, SecondsDiff(InSecondsDiff)
+		, SubjectName(InSubjectName)
+		, SerializedType(InSerializedType)
+		, Guid(InGuid)
+	{
+	}
+
+	friend FArchive& operator<<(FArchive& Ar, FLiveLinkFileHeader& Header)
+	{
+		Ar.UsingCustomVersion(FLiveLinkCustomVersion::GUID);
+
+		Ar << Header.SerializedType;
+		Ar << Header.bIsManifest;
+		Ar << Header.Guid;
+		Ar << Header.SecondsDiff;
+		Ar << Header.SubjectName;
+
+		if (Ar.IsLoading())
+		{
+			if (Ar.CustomVer(FLiveLinkCustomVersion::GUID) < FLiveLinkCustomVersion::NewLiveLinkRoleSystem)
+			{
+				Ar << Header.CurveNames;
+				FLiveLinkRefSkeleton::StaticStruct()->SerializeItem(Ar, (void*)& Header.RefSkeleton, nullptr);
+			}
+		}
+
+		return Ar;
+	}
+
+	bool bIsManifest;
+	double SecondsDiff;
+	FName SubjectName;
+	FName SerializedType;
+	FGuid Guid;
+
+	TArray<FName> CurveNames;
+	FLiveLinkRefSkeleton RefSkeleton;
+};
+PRAGMA_ENABLE_DEPRECATION_WARNINGS
+
diff --git a/Engine/Source/Runtime/SerializedRecorderInterface/Public/Serializers/MovieSceneManifestSerialization.h b/Engine/Source/Runtime/SerializedRecorderInterface/Public/Serializers/MovieSceneManifestSerialization.h
new file mode 100644
--- /dev/null
+++ b/Engine/Source/Runtime/SerializedRecorderInterface/Public/Serializers/MovieSceneManifestSerialization.h
@@ -0,0 +1,73 @@
+// Copyright Epic Games, Inc. All Rights Reserved.
+
+#pragma once
+
+#include "Misc/FrameRate.h"
+#include "UObject/UnrealType.h"
+#include "MovieSceneSectionSerialization.h"
+
+
+struct FManifestFileHeader
+{
+	static const int32 cVersion = 1;
+
+	FManifestFileHeader() : Version(cVersion)
+	{
+	}
+
+	FManifestFileHeader(const FString& InName, const FName& InSerializedType, const FGuid& InGuid) 
+		: Version(cVersion)
+		, SerializedType(InSerializedType)
+		, Guid(InGuid)
+		, Name(InName)
+
+	{
+	}
+
+	friend FArchive& operator<<(FArchive& Ar, FManifestFileHeader& Header)
+	{
+		Ar << Header.Version;
+		Ar << Header.SerializedType;
+		Ar << Header.Guid;
+		Ar << Header.Name;
+
+		return Ar;
+	}
+
+	int32 Version;
+	FName SerializedType;
+	FGuid Guid;
+	FString Name;
+};
+
+
+//since we can't use function templates with TFunctions need to use a variant or union.
+struct FManifestProperty
+{
+	FManifestProperty() = default;
+	FManifestProperty(const FString &InObjectName, const FName& InSerializedType, const FGuid& InGuid) :
+		UObjectName(InObjectName)
+		, SerializedType(InSerializedType)
+		, Guid(InGuid)
+	{}
+	friend FArchive& operator<<(FArchive& Ar, FManifestProperty& Property)
+	{
+
+		Ar << Property.SerializedType;
+		Ar << Property.Guid;
+		Ar << Property.UObjectName;
+		return Ar;
+	}
+
+
+	//DATA
+
+	FString UObjectName;
+	FName SerializedType;
+	FGuid Guid;
+
+};
+
+
+using FManifestSerializedFrame = TMovieSceneSerializedFrame<FManifestProperty>;
+using FManifestSerializer = TMovieSceneSerializer<FManifestFileHeader, FManifestProperty>;
\ No newline at end of file
diff --git a/Engine/Source/Runtime/SerializedRecorderInterface/Public/Serializers/MovieScenePropertySerialization.h b/Engine/Source/Runtime/SerializedRecorderInterface/Public/Serializers/MovieScenePropertySerialization.h
new file mode 100644
--- /dev/null
+++ b/Engine/Source/Runtime/SerializedRecorderInterface/Public/Serializers/MovieScenePropertySerialization.h
@@ -0,0 +1,183 @@
+// Copyright Epic Games, Inc. All Rights Reserved.
+
+#pragma once
+
+#include "Misc/FrameRate.h"
+#include "UObject/StrProperty.h"
+#include "UObject/UnrealType.h"
+#include "UObject/EnumProperty.h"
+#include "MovieSceneSectionSerialization.h"
+
+enum class ESerializedPropertyType : uint8 {
+
+	BoolType,
+	ByteType,
+	EnumType,
+	FloatType,
+	DoubleType,
+	Vector3fType,
+	Vector3dType,
+	ColorType,
+	IntegerType,
+	StringType,
+	LinearColorType
+};
+
+struct FPropertyFileHeader
+{
+
+	static const int32 cVersion = 1;
+
+	FPropertyFileHeader() : Version(cVersion)
+	{
+	}
+
+
+	FPropertyFileHeader(const FFrameRate &InFrameRate, const FName& InSerializedType, const FGuid& InGuid)
+		: Version(cVersion)
+		, SerializedType(InSerializedType)
+		, Guid(InGuid)
+		, TickResolution(InFrameRate)
+
+	{
+	}
+
+	friend FArchive& operator<<(FArchive& Ar, FPropertyFileHeader& Header)
+	{
+		Ar << Header.Version;
+		Ar << Header.SerializedType;
+		Ar << Header.Guid;
+		Ar << Header.TickResolution.Numerator;
+		Ar << Header.TickResolution.Denominator;
+		Ar << Header.PropertyName;
+		Ar << Header.PropertyType;
+		Ar << Header.TrackDisplayName;
+
+		return Ar;
+	}
+
+	void SetProperty(const FProperty* Property, const FName& InPropertyName)
+	{
+		if (Property != nullptr)
+		{
+			PropertyName = InPropertyName;
+
+		#if WITH_EDITOR
+			TrackDisplayName = *Property->GetDisplayNameText().ToString();
+		#endif // WITH_EDITOR
+
+			if (Property->IsA<FBoolProperty>())
+			{
+				PropertyType = (ESerializedPropertyType::BoolType);
+			}
+			else if (Property->IsA<FByteProperty>())
+			{
+				PropertyType = (ESerializedPropertyType::ByteType);
+			}
+			else if (Property->IsA<FEnumProperty>())
+			{
+				PropertyType = (ESerializedPropertyType::EnumType);
+			}
+			else if (Property->IsA<FFloatProperty>())
+			{
+				PropertyType = (ESerializedPropertyType::FloatType);
+			}
+			else if (Property->IsA<FDoubleProperty>())
+			{
+				PropertyType = (ESerializedPropertyType::DoubleType);
+			}
+			else if (const FStructProperty* StructProperty = CastField<const FStructProperty>(Property))
+			{
+				// LWC_TODO: vector types
+				if (StructProperty->Struct->GetFName() == NAME_Vector3f)
+				{
+					PropertyType = (ESerializedPropertyType::Vector3fType);
+				}
+				else if (StructProperty->Struct->GetFName() == NAME_Vector3d
+						|| StructProperty->Struct->GetFName() == NAME_Vector
+						)
+				{
+					PropertyType = (ESerializedPropertyType::Vector3dType);
+				}
+				else if (StructProperty->Struct->GetFName() == NAME_Color)
+				{
+					PropertyType = (ESerializedPropertyType::ColorType);
+				}
+				else if (StructProperty->Struct->GetFName() == NAME_LinearColor)
+				{
+					PropertyType = (ESerializedPropertyType::LinearColorType);
+				}
+			}
+			else if (Property->IsA<FIntProperty>())
+			{
+				PropertyType = (ESerializedPropertyType::IntegerType);
+			}
+			else if (Property->IsA<FStrProperty>())
+			{
+				PropertyType = (ESerializedPropertyType::StringType);
+			}
+		}
+	}
+
+
+	//DATA
+	int32 Version;
+	FName SerializedType;
+	FGuid Guid;
+	FFrameRate TickResolution;
+	FName PropertyName;
+	ESerializedPropertyType PropertyType;
+	FString TrackDisplayName;
+};
+
+
+template <typename PropertyType>
+struct FSerializedProperty
+{
+	FSerializedProperty() = default;
+
+	friend FArchive& operator<<(FArchive& Ar, FSerializedProperty& Property)
+	{
+		Ar << Property.Time;
+		Ar << Property.Value;
+		return Ar;
+	}
+
+	FFrameNumber Time;
+	PropertyType Value;
+
+};
+
+using FPropertySerializedBool = FSerializedProperty<bool>;
+using FPropertySerializedBoolFrame = TMovieSceneSerializedFrame<FSerializedProperty<bool>>;
+
+using FPropertySerializedByte = FSerializedProperty<uint8>;
+using FPropertySerializedByteFrame = TMovieSceneSerializedFrame<FSerializedProperty<uint8>>;
+
+using FPropertySerializedEnum = FSerializedProperty<int64>;
+using FPropertySerializedEnumFrame = TMovieSceneSerializedFrame<FSerializedProperty<int64>>;
+using FPropertySerializerEnum = TMovieSceneSerializer<FPropertyFileHeader, FSerializedProperty<int64>>;
+
+using FPropertySerializedDouble = FSerializedProperty<double>;
+using FPropertySerializedDoubleFrame = TMovieSceneSerializedFrame<FSerializedProperty<double>>;
+
+using FPropertySerializedFloat = FSerializedProperty<float>;
+using FPropertySerializedFloatFrame = TMovieSceneSerializedFrame<FSerializedProperty<float>>;
+
+using FPropertySerializedVector3f = FSerializedProperty<FVector3f>;
+using FPropertySerializedVector3fFrame = TMovieSceneSerializedFrame<FSerializedProperty<FVector3f>>;
+
+using FPropertySerializedVector3d = FSerializedProperty<FVector3d>;
+using FPropertySerializedVector3dFrame = TMovieSceneSerializedFrame<FSerializedProperty<FVector3d>>;
+
+using FPropertySerializedColor = FSerializedProperty<FColor>;
+using FPropertySerializedColorFrame = TMovieSceneSerializedFrame<FSerializedProperty<FColor>>;
+
+using FPropertySerializedInteger = FSerializedProperty<int32>;
+using FPropertySerializedIntegerFrame = TMovieSceneSerializedFrame<FSerializedProperty<int32>>;
+
+using FPropertySerializedString = FSerializedProperty<FString>;
+using FPropertySerializedStringFrame = TMovieSceneSerializedFrame<FSerializedProperty<FString>>;
+
+using FPropertySerializedLinearColor = FSerializedProperty<FLinearColor>;
+using FPropertySerializedLinearColorFrame = TMovieSceneSerializedFrame<FSerializedProperty<FLinearColor>>;
\ No newline at end of file
diff --git a/Engine/Source/Runtime/SerializedRecorderInterface/Public/Serializers/MovieSceneSectionSerialization.h b/Engine/Source/Runtime/SerializedRecorderInterface/Public/Serializers/MovieSceneSectionSerialization.h
new file mode 100644
--- /dev/null
+++ b/Engine/Source/Runtime/SerializedRecorderInterface/Public/Serializers/MovieSceneSectionSerialization.h
@@ -0,0 +1,1055 @@
+// Copyright Epic Games, Inc. All Rights Reserved.
+
+#pragma once
+
+#include "Algo/BinarySearch.h"
+#include "Async/Async.h"
+#include "Async/TaskGraphInterfaces.h"
+#include "Containers/Array.h"
+#include "Containers/ContainersFwd.h"
+#include "Containers/LockFreeList.h"
+#include "Containers/Map.h"
+#include "Containers/UnrealString.h"
+#include "GenericPlatform/GenericPlatformFile.h"
+#include "HAL/CriticalSection.h"
+#include "HAL/FileManager.h"
+#include "HAL/FileManagerGeneric.h"
+#include "HAL/PlatformFileManager.h"
+#include "HAL/PlatformProcess.h"
+#include "HAL/PlatformProperties.h"
+#include "HAL/Runnable.h"
+#include "HAL/RunnableThread.h"
+#include "HAL/ThreadSafeBool.h"
+#include "Internationalization/Internationalization.h"
+#include "Internationalization/Text.h"
+#include "Logging/LogCategory.h"
+#include "Logging/LogMacros.h"
+#include "Math/UnrealMathUtility.h"
+#include "Misc/App.h"
+#include "Misc/AssertionMacros.h"
+#include "Misc/DateTime.h"
+#include "Misc/FrameRate.h"
+#include "Misc/Guid.h"
+#include "Misc/NetworkVersion.h"
+#include "Misc/Paths.h"
+#include "Misc/ScopeLock.h"
+#include "Serialization/Archive.h"
+#include "Serialization/ArchiveProxy.h"
+#include "Serialization/CustomVersion.h"
+#include "Serialization/NameAsStringProxyArchive.h"
+#include "Templates/Function.h"
+#include "Templates/Less.h"
+#include "Templates/SharedPointer.h"
+#include "Templates/Tuple.h"
+#include "Templates/UniquePtr.h"
+#include "Templates/UnrealTemplate.h"
+#include "Trace/Detail/Channel.h"
+
+SERIALIZEDRECORDERINTERFACE_API DECLARE_LOG_CATEGORY_EXTERN(MovieSceneSerialization, Verbose, All);
+
+namespace MovieSceneSerializationNamespace
+{
+	extern SERIALIZEDRECORDERINTERFACE_API const int64 InvalidOffset;
+	extern SERIALIZEDRECORDERINTERFACE_API const float SerializerSleepTime;
+	extern SERIALIZEDRECORDERINTERFACE_API bool bAutoSerialize;
+};
+
+template <typename FrameData>
+struct TMovieSceneSerializedFrame
+{
+	uint64 FrameId;
+	FrameData  Frame;
+
+	friend FArchive& operator<<(FArchive& Ar, TMovieSceneSerializedFrame& InFrame)
+	{
+		Ar << InFrame.FrameId;
+		Ar << InFrame.Frame;
+		return Ar;
+	}
+};
+
+
+/** Header (non-tagged property serialized) for a frame */
+struct FMovieSceneSerializedFrameHeader
+{
+	FMovieSceneSerializedFrameHeader()
+		: NextFrameOffset(MovieSceneSerializationNamespace::InvalidOffset)
+		, PrevFrameOffset(MovieSceneSerializationNamespace::InvalidOffset)
+		, FrameId(0)
+	{}
+
+	FMovieSceneSerializedFrameHeader(int64 InNextFrameOffset, int64 InPrevFrameOffset, uint64 InFrameId)
+		: NextFrameOffset(InNextFrameOffset)
+		, PrevFrameOffset(InPrevFrameOffset)
+		, FrameId(InFrameId)
+	{}
+
+	/** Offset within the archive that the next frame can be found. -1 if this is the last frame. */
+	int64 NextFrameOffset;
+
+	/** Offset within the archive that the next frame can be found. -1 if this is the first frame. */
+	int64 PrevFrameOffset;
+
+	/** ID of this frame */
+	uint64 FrameId;
+
+
+	friend FArchive& operator<<(FArchive& Ar, FMovieSceneSerializedFrameHeader& Header)
+	{
+		Ar << Header.NextFrameOffset;
+		Ar << Header.PrevFrameOffset;
+		Ar << Header.FrameId;
+
+		return Ar;
+	}
+};
+
+class FMovieSceneArchiveFileReader : public FArchiveFileReaderGeneric
+{
+public:
+	FMovieSceneArchiveFileReader(IFileHandle* InHandle, const TCHAR* InFilename)
+		: FArchiveFileReaderGeneric(InHandle, InFilename, InHandle->Size())
+	{}
+
+	void BeginArchiving()
+	{
+		Size = Handle->Size();
+		Pos = Handle->Tell();
+		BufferBase = 0;
+		BufferArray.Reset();
+	}
+
+	/** FArchiveFileReaderGeneric interface*/
+	virtual void CloseLowLevel() override
+	{
+		// Release the handle as this archive doesnt actually own it
+		(void)Handle.Release();
+	}
+};
+
+class FMovieSceneArchiveFileWriter : public FArchiveFileWriterGeneric
+{
+public:
+	FMovieSceneArchiveFileWriter(IFileHandle* InHandle, const TCHAR* InFilename)
+		: FArchiveFileWriterGeneric(InHandle, InFilename, InHandle->Tell())
+	{}
+
+	void BeginArchiving()
+	{
+		Pos = Handle->Tell();
+		BufferArray.Reset();
+	}
+
+	/** FArchiveFileWriterGeneric interface*/
+	virtual bool CloseLowLevel() override
+	{
+		// Release the handle as this archive doesnt actually own it
+		(void)Handle.Release();
+		return true;
+	}
+};
+
+
+
+
+/** Information about a captured session */
+struct FMovieSceneSessionInfo
+{
+	FMovieSceneSessionInfo()
+		: BuildVersion(0)
+	{}
+
+	FMovieSceneSessionInfo(const FGuid& InInstanceId, FDateTime InTimeStamp, int32 InBuildVersion, const FString& InDeviceName, const FString& InPlatformName, const FString& InInstanceName, const FString& InBuildDate)
+		: InstanceId(InInstanceId)
+		, TimeStamp(InTimeStamp)
+		, BuildVersion(InBuildVersion)
+		, DeviceName(InDeviceName)
+		, PlatformName(InPlatformName)
+		, InstanceName(InInstanceName)
+		, BuildDate(InBuildDate)
+	{}
+
+	/** The ID of the session instance the file was captured with */
+	FGuid InstanceId;
+
+	/** The timestamp of this capture at creation time. @see FDateTime::UtcNow() */
+	FDateTime TimeStamp;
+
+	/** The engine version the file was captured with. @see ISessionInstanceInfo */
+	int32 BuildVersion;
+
+	/** Name of the device the file was captured on. @see ISessionInstanceInfo */
+	FString DeviceName;
+
+	/** Name of the platform the file was captured on. @see ISessionInstanceInfo */
+	FString PlatformName;
+
+	/** Name of the engine instance the file was captured on. @see ISessionInstanceInfo */
+	FString InstanceName;
+
+	/** The build date for the capture. @see ISessionInstanceInfo */
+	FString BuildDate;
+
+	friend FArchive& operator<<(FArchive& Ar, FMovieSceneSessionInfo& InOutSessionInfo)
+	{
+		Ar << InOutSessionInfo.InstanceId;
+		Ar << InOutSessionInfo.TimeStamp;
+		Ar << InOutSessionInfo.BuildVersion;
+		Ar << InOutSessionInfo.DeviceName;
+		Ar << InOutSessionInfo.PlatformName;
+		Ar << InOutSessionInfo.InstanceName;
+		Ar << InOutSessionInfo.BuildDate;
+
+		return Ar;
+	}
+};
+
+
+
+
+/** Archive context maintained when the serializer is running */
+class FMovieSceneSerializerContext
+{
+public:
+
+	FMovieSceneSerializerContext(IFileHandle* InHandle, const FString& InFilename)
+		: Filename(InFilename)
+		, Handle(InHandle)
+		, ArReaderInner(InHandle, *Filename)
+		, ArWriterInner(InHandle, *Filename)
+		, ArReader(ArReaderInner)
+		, ArWriter(ArWriterInner)
+		, LastFrameWritePos(MovieSceneSerializationNamespace::InvalidOffset)
+		, MinFrameId(0xffffffffffffffff)
+		, MaxFrameId(0)
+	{}
+
+	~FMovieSceneSerializerContext() = default;
+
+	void Close()
+	{
+		ArReader.Close();
+		ArWriter.Close();
+		delete Handle;
+		Handle = nullptr;
+	}
+
+
+public:
+	
+	/** Filename */
+	FString Filename;
+
+	/** Read/write handles */
+	IFileHandle* Handle;
+
+	/** Inner file reader/writer archives */
+	FMovieSceneArchiveFileReader ArReaderInner;
+	FMovieSceneArchiveFileWriter ArWriterInner;
+
+	/** Reader/writer archives */
+	FNameAsStringProxyArchive ArReader;
+	FNameAsStringProxyArchive ArWriter;
+
+	/** Offset within the archive that the last frame was written */
+	int64 LastFrameWritePos;
+
+	/** Frame IDs used to access specific times quickly */
+	TArray<TPair<uint64, int64>> FrameIdToFrameOffset;
+
+	/** The min frame ID contained in the file */
+	uint64 MinFrameId;
+
+	/** The max frame ID contained in the file */
+	uint64 MaxFrameId;
+
+};
+
+
+typedef TFunction<void(FMovieSceneSerializerContext& /*InContext*/)> FSerializerCommand;
+
+struct FTempCustomVersion
+{
+	enum Type
+	{
+		// Initial version
+		FirstVersion = 0,
+
+		// -----new versions can be added above this line-------------------------------------------------
+		VersionPlusOne,
+		LatestVersion = VersionPlusOne - 1
+	};
+
+	// The GUID for this custom version number
+	SERIALIZEDRECORDERINTERFACE_API const static FGuid GUID;
+private:
+	FTempCustomVersion() {}
+};
+
+
+template <typename HeaderData, typename FrameData>
+struct FContextAndCommands
+{
+	FContextAndCommands(IFileHandle* InHandle, const FString& InFilename)
+		: Context(InHandle, InFilename)
+	{
+
+	}
+	~FContextAndCommands()
+	{
+		while (FSerializerCommand* Command = ReadCommands.Pop())
+		{
+			delete Command;
+		}
+
+		while (FSerializerCommand* Command = WriteCommands.Pop())
+		{
+			delete Command;
+		}
+
+		while (FSerializerCommand* Command = FreeCommands.Pop())
+		{
+			delete Command;
+		}
+	}
+	/** Unique ID of this Guid, key for a map of Context*/
+	FGuid Guid;
+
+	FMovieSceneSerializerContext Context;
+
+	/** Read command queue - pushed on any thread, popped on the serializer worker thread */
+	TLockFreePointerListFIFO<FSerializerCommand, 0> ReadCommands;
+
+	/** Write command queue - pushed on any thread, popped on the serializer worker thread */
+	TLockFreePointerListFIFO<FSerializerCommand, 0> WriteCommands;
+
+	/** Free command queue - pushed on any thread, popped on any thread */
+	TLockFreePointerListFIFO<FSerializerCommand, 0> FreeCommands;
+	
+	void Close()
+	{
+		Context.Close();
+	}
+
+	void Run()
+	{
+		// Pop and run our commands
+		TArray<FSerializerCommand*> LocalReadCommands;
+		ReadCommands.PopAll(LocalReadCommands);
+
+		TArray<FSerializerCommand*> LocalWriteCommands;
+		WriteCommands.PopAll(LocalWriteCommands);
+
+		if (LocalReadCommands.Num() > 0)
+		{
+			Context.ArReaderInner.BeginArchiving();
+
+			for (FSerializerCommand* Command : LocalReadCommands)
+			{
+				if (Command != nullptr)
+				{
+					(*Command)(Context);
+					FreeCommands.Push(Command);
+				}
+			}
+
+			Context.ArReader.Flush();
+		}
+
+		if (LocalWriteCommands.Num() > 0)
+		{
+			Context.ArWriterInner.BeginArchiving();
+
+			for (FSerializerCommand* Command : LocalWriteCommands)
+			{
+				if (Command != nullptr)
+				{
+					(*Command)(Context);
+					FreeCommands.Push(Command);
+				}
+			}
+
+			Context.ArWriter.Flush();
+		}
+	}
+
+
+	bool CheckHeader(bool bInFileExisted, TArray<FMovieSceneSessionInfo>& SessionInfos, HeaderData& InOutHeader, int64& ReadHeaderPos,  FText& OutFailReason)
+	{
+
+		Context.ArReader.UsingCustomVersion(FTempCustomVersion::GUID);
+		Context.ArWriter.UsingCustomVersion(FTempCustomVersion::GUID);
+
+		if (bInFileExisted)
+		{
+			// If we are reading an existing file, check the version
+			int32 ArchiveVer = INT_MAX;
+			Context.ArReader << ArchiveVer;
+
+			if (ArchiveVer > FTempCustomVersion::LatestVersion)
+			{
+				OutFailReason = FText::Format(NSLOCTEXT("MovieSceneSerializer", "FileOpenFailedIncorrectVersion", "Cannot open file.\nFile version: {0}.\nMax supported version: {1}."), FText::AsNumber(ArchiveVer), FText::AsNumber((int32)FTempCustomVersion::LatestVersion));
+				UE_LOGF(MovieSceneSerialization, Error, "Cannot open Movie Scene Serialization cache %ls. File version: %d. Max supported version: %d", *Context.Filename, ArchiveVer, FTempCustomVersion::LatestVersion);
+				Context.Close();
+				return false;
+			}
+
+
+			int32 NumSessions = SessionInfos.Num();
+			Context.ArReader << NumSessions;
+
+			if (Context.ArReader.IsLoading())
+			{
+				SessionInfos.SetNum(NumSessions);
+			}
+
+			for (int32 SessionIndex = 0; SessionIndex < NumSessions; ++SessionIndex)
+			{
+				Context.ArReader << SessionInfos[SessionIndex];
+			}
+
+			Context.ArReader << InOutHeader;
+
+			Context.ArReader.Flush();
+		}
+		else
+		{
+			int32 ArchiveVer = Context.ArWriter.CustomVer(FTempCustomVersion::GUID);
+			Context.ArWriter << ArchiveVer;
+
+
+			int32 NumSessions = SessionInfos.Num();
+			Context.ArWriter << NumSessions;
+			for (int32 SessionIndex = 0; SessionIndex < NumSessions; ++SessionIndex)
+			{
+				Context.ArWriter << SessionInfos[SessionIndex];
+			}
+
+			Context.ArWriter << InOutHeader;
+			Context.ArWriter.Flush();
+		}
+		ReadHeaderPos = Context.ArReader.Tell();
+
+		return true;
+	}
+
+	
+	/**
+	* Gets a free command from the free list, if any. Otherwise allocates a new command. This means
+	* that the number of commands grows unbounded to the high watermark of all commands.
+	*/
+	template <typename... ArgsType>
+	FORCEINLINE FSerializerCommand* GetFreeCommand(ArgsType&&... Args)
+	{
+		FSerializerCommand* Command = FreeCommands.Pop();
+		if (Command == nullptr)
+		{
+			Command = new FSerializerCommand(Forward<ArgsType>(Args)...);
+		}
+		else
+		{
+			*Command = FSerializerCommand(Forward<ArgsType>(Args)...);
+		}
+
+		return Command;
+	}
+
+	template <typename... ArgsType>
+	FORCEINLINE void AddReadCommand(ArgsType&&... Args)
+	{
+		ReadCommands.Push(GetFreeCommand(Forward<ArgsType>(Args)...));
+	}
+
+	template <typename... ArgsType>
+	FORCEINLINE void AddWriteCommand(ArgsType&&... Args)
+	{
+		WriteCommands.Push(GetFreeCommand(Forward<ArgsType>(Args)...));
+	}
+};
+
+/** Runnable to run threaded I/O */
+template <typename HeaderData, typename FrameData>
+class FMovieSceneSerializerRunnable : public FRunnable
+{
+public:
+	FMovieSceneSerializerRunnable()
+		: bRunning(true), bStopping(false)
+	{}
+
+	~FMovieSceneSerializerRunnable()
+	{
+		check(!bRunning);
+
+	
+	}
+
+	void AddContext(const FGuid& InGuid, IFileHandle* InHandle, const FString& InFileName)
+	{
+		FScopeLock Lock(&ContextMapCriticalSection);
+		TSharedPtr<FContextAndCommands<HeaderData, FrameData>> Context = MakeShared<FContextAndCommands<HeaderData, FrameData>>(InHandle, InFileName);
+		ContextMap.Add(InGuid, Context);
+	}
+
+	/** FRunnable interface */
+	virtual uint32 Run() override
+	{
+		while (bRunning)
+		{
+			{
+				FScopeLock Lock(&ContextMapCriticalSection);
+
+				for (auto& Entry : ContextMap)
+				{
+					TSharedPtr <FContextAndCommands<HeaderData, FrameData>> & Context = Entry.Value;
+					Context->Run();
+				}
+			}
+
+			if (bStopping)
+			{
+				bRunning = false;
+			}
+			FPlatformProcess::Sleep(MovieSceneSerializationNamespace::SerializerSleepTime);
+		}
+
+		CloseContexts();
+		return 0;
+	}
+
+	virtual void Stop() override
+	{
+		bStopping = true;
+
+	}
+
+	bool CheckHeader(const FGuid& InGuid, bool bInFileExisted, TArray<FMovieSceneSessionInfo>& SessionInfos, HeaderData& InOutHeader, int64& ReadHeaderPos,  FText& OutFailReason)
+	{
+		FScopeLock Lock(&ContextMapCriticalSection);
+		TSharedPtr<FContextAndCommands<HeaderData, FrameData >>* Context = ContextMap.Find(InGuid);
+		if (Context)
+		{
+			return (*Context)->CheckHeader(bInFileExisted, SessionInfos, InOutHeader, ReadHeaderPos, OutFailReason);
+		}
+
+		return false;
+	}
+
+	FContextAndCommands<HeaderData,FrameData> & GetContext(const FGuid& InGuid)
+	{
+		FScopeLock Lock(&ContextMapCriticalSection);
+		TSharedPtr<FContextAndCommands<HeaderData, FrameData >>* Context = ContextMap.Find(InGuid);
+		return *(Context->Get());
+	}
+	
+	void Close(const FGuid& InGuid)
+	{
+		FScopeLock Lock(&ContextMapCriticalSection);
+		TSharedPtr<FContextAndCommands<HeaderData, FrameData >>* Context = ContextMap.Find(InGuid);
+		if (Context)
+		{
+			(*Context)->Close();
+			ContextMap.Remove(InGuid);
+		}
+	}
+
+	void CloseContexts()
+	{
+		FScopeLock Lock(&ContextMapCriticalSection);
+
+		for (auto& Entry : ContextMap)
+		{
+			TSharedPtr <FContextAndCommands<HeaderData, FrameData>> & Context = Entry.Value;
+			Context->Close();
+		}
+		ContextMap.Empty();
+
+	}
+	/** The contexts we are using */
+	TMap<FGuid, TSharedPtr<FContextAndCommands<HeaderData, FrameData>>> ContextMap;
+
+	/** Thread safe bool for running the thread */
+	FThreadSafeBool bRunning;
+
+	/** Thread safe bool for stopping the thread */
+	FThreadSafeBool bStopping;
+
+	/** Read/write access flags */
+	FThreadSafeBool bOpenForRead;
+	FThreadSafeBool bOpenForWrite;
+
+	/**  Lock */
+	FCriticalSection ContextMapCriticalSection;
+
+};
+
+/** A class to asynchronously read and write to serialized frame debugger data files */
+template <typename HeaderData,typename FrameData>
+class  TMovieSceneSerializer
+{
+public:
+
+	TMovieSceneSerializer();
+
+	/** Get whether file exists */
+	bool DoesFileExist(const FString& InFileName);
+
+	/** Get the local capture directory */
+	FString GetLocalCaptureDir();
+
+	/** Set the local capture directory */
+	void SetLocalCaptureDir(const FString& PathIn);
+	
+	/**
+	* Open the serializer for writing
+	* @return false if the operation was not successful with regards to opening the file, if we have saving turned off it will still return true.
+	*/
+	bool OpenForWrite(const FString& InFilename, HeaderData& Header, FText& OutFailReason);
+	/** 
+	 * Open the serializer for reading
+	 * @return false if the operation was not successful
+	 */
+
+	bool OpenForRead(const FString& InFileName, HeaderData& OutHeader,  FText& OutFailReason);
+
+	/** Close the serializer, flushes any commands and shuts down threads */
+	void Close();
+
+	/** Check whether this serializer is currently open */
+	bool IsOpen() const;
+
+	/** Write one frame data out to disk */
+	void WriteFrameData(uint64 InFrameId, const FrameData &InFrame);
+
+	/** Write array of frame data out to disk */
+	void WriteFrameData(uint64 InFrameId, const TArray<FrameData> &InFrames);
+
+	/** Read a frame range. Completion callback will be called on the game thread. */
+	void ReadFramesAtFrameRange(uint64 InStartFrameId, uint64 InEndFrameId, TFunction<void()> InCompletionCallback);
+
+	/** Query the range of data in the file. Completion callback will be called on the game thread. */
+	void GetDataRanges(TFunction<void(uint64 /*InMinFrameId*/, uint64 /*InMaxFrameId*/)> InCompletionCallback);
+
+	/** Get the number of sessions that were captured to this file */
+	int32 GetNumSessions() { return SessionInfos.Num(); }
+
+	/** Get the engine version the file was captured with. @see ISessionInstanceInfo */
+	FGuid GetSessionId(int32 InSessionIndex) const { return SessionInfos[InSessionIndex].InstanceId; }
+
+	/** Get the engine version the file was captured with. @see ISessionInstanceInfo */
+	int32 GetBuildVersion(int32 InSessionIndex) const { return SessionInfos[InSessionIndex].BuildVersion; }
+
+	/** Get the name of the device the file was captured on. @see ISessionInstanceInfo */
+	FString GetDeviceName(int32 InSessionIndex) const { return SessionInfos[InSessionIndex].DeviceName; }
+
+	/** Get the name of the platform the file was captured on. @see ISessionInstanceInfo */
+	FString GetPlatformName(int32 InSessionIndex) const { return SessionInfos[InSessionIndex].PlatformName; }
+
+	/** Get the name of the engine instance the file was captured on. @see ISessionInstanceInfo */
+	FString GetInstanceName(int32 InSessionIndex) const { return SessionInfos[InSessionIndex].InstanceName; }
+
+	/** Get the build date for the capture. @see ISessionInstanceInfo */
+	FString GetBuildDate(int32 InSessionIndex) const { return SessionInfos[InSessionIndex].BuildDate; }
+
+	/** How many frames have been written */
+	int FramesWritten;
+
+	/** Data from a read.  Be carefule if reading multiple times this data will get replaced and modified, so need to 
+	come up with different way to get data in that case*/
+	TArray<TMovieSceneSerializedFrame<FrameData>> ResultData;
+
+	/** Thread we run on */
+    static  SERIALIZEDRECORDERINTERFACE_API FRunnableThread* Thread;
+
+	/** Runnable that the thread runs */
+	static  SERIALIZEDRECORDERINTERFACE_API FMovieSceneSerializerRunnable<HeaderData, FrameData>* Runnable;
+private:
+
+
+	/** The session infos */
+	TArray<FMovieSceneSessionInfo> SessionInfos;
+
+	/** Local Directory*/
+	FString LocalCaptureDir;
+
+	/** Unique guid for this serializer */
+	FGuid Guid;
+};
+
+
+template<typename HeaderData, typename FrameData>
+TMovieSceneSerializer<HeaderData, FrameData>::TMovieSceneSerializer()
+	: FramesWritten(0)
+	, LocalCaptureDir(FPaths::ProjectSavedDir())
+	, Guid(FGuid::NewGuid())
+{
+	
+}
+template<typename HeaderData, typename FrameData>
+FString TMovieSceneSerializer<HeaderData, FrameData>::GetLocalCaptureDir()
+{
+	return LocalCaptureDir;
+}
+template<typename HeaderData, typename FrameData>
+void  TMovieSceneSerializer<HeaderData, FrameData>::SetLocalCaptureDir(const FString& PathIn)
+{
+	LocalCaptureDir = PathIn;
+}
+template<typename HeaderData, typename FrameData>
+bool TMovieSceneSerializer<HeaderData, FrameData>::DoesFileExist(const FString& AbsoluteFilePath)
+{
+
+	return IFileManager::Get().FileExists(*AbsoluteFilePath);
+}
+
+template<typename HeaderData, typename FrameData>
+bool TMovieSceneSerializer<HeaderData, FrameData>::OpenForWrite(const FString& InFilename, HeaderData& Header, FText& OutFailReason)
+{
+	if (MovieSceneSerializationNamespace::bAutoSerialize)
+	{
+		// Get absolute file path
+		FString SaveDirectory = GetLocalCaptureDir();
+		FString AbsoluteFilePath = SaveDirectory + "/" + InFilename;
+		FramesWritten = 0;
+		// Create file handle and hand off to runnable thread
+		IFileHandle* Handle = FPlatformFileManager::Get().GetPlatformFile().OpenWrite(*AbsoluteFilePath, false, true);
+		if (Handle != nullptr)
+		{
+			UE_LOGF(MovieSceneSerialization, Display, "Opened Data cache: %ls", *AbsoluteFilePath);
+
+			// Create runnable
+			if (!Runnable)
+			{
+				Runnable = new FMovieSceneSerializerRunnable<HeaderData, FrameData>();
+			}
+			Runnable->AddContext(Guid, Handle, AbsoluteFilePath);
+			Runnable->bOpenForRead = true;
+			Runnable->bOpenForWrite = true;
+			FMovieSceneSessionInfo SessionInfo(FApp::GetInstanceId(), FDateTime::UtcNow(), FNetworkVersion::GetNetworkCompatibleChangelist(), FPlatformProcess::ComputerName(),
+				FPlatformProperties::PlatformName(), FApp::GetInstanceName(), FApp::GetBuildDate());
+			SessionInfos.Emplace(SessionInfo);
+			int64 ReadHeaderPos;
+
+			if (Runnable->CheckHeader(Guid, false, SessionInfos, Header, ReadHeaderPos, OutFailReason))
+			{
+				// Create thread
+				if (!Thread)
+				{
+					FString ThreadName = (Header.SerializedType.ToString());
+					Thread = FRunnableThread::Create(Runnable, *ThreadName);
+					UE_LOGF(MovieSceneSerialization, Display, "Created Thread For Type: %ls", *ThreadName);
+				}
+
+				return true;
+			}
+			else
+			{
+				Runnable->Close(Guid);
+				/*
+				Runnable->Stop();
+				delete Runnable;
+				Runnable = nullptr;
+				*/
+				return false;
+			}
+		}
+		else
+		{
+			OutFailReason = NSLOCTEXT("MovieSceneSerializer", "FileOpenFailedToCreateArchive", "Cannot open file.\nFailed to create archive.");
+
+			UE_LOGF(MovieSceneSerialization, Error, "Cannot open Movie Scene Serialization cache %ls. Failed to create archive.", *InFilename);
+			return false;
+		}
+	}
+	else
+	{
+		if (Runnable)
+		{
+			Runnable->bOpenForWrite = false;
+		}
+		return true;
+	}
+}
+
+template<typename HeaderData, typename FrameData>
+bool TMovieSceneSerializer<HeaderData, FrameData>::OpenForRead(const FString& AbsoluteFilePath, HeaderData& OutHeader, FText& OutFailReason)
+{
+
+	const bool bFileExists = IFileManager::Get().FileExists(*AbsoluteFilePath);
+
+	// Create file handle and hand off to runnable thread
+	IFileHandle* Handle = bFileExists ? FPlatformFileManager::Get().GetPlatformFile().OpenRead(*AbsoluteFilePath) : nullptr;
+	if (Handle != nullptr)
+	{
+		UE_LOGF(MovieSceneSerialization, Display, "Opened Movie Scene Serialization cache: %ls", *AbsoluteFilePath);
+
+		// Create runnable
+		if (!Runnable)
+		{
+			Runnable = new FMovieSceneSerializerRunnable<HeaderData, FrameData>();
+		}
+		Runnable->AddContext(Guid, Handle, AbsoluteFilePath);
+
+		Runnable->bOpenForRead = true;
+		Runnable->bOpenForWrite = false;
+		int64 ReadHeaderPos;
+
+		if (Runnable->CheckHeader(Guid,bFileExists, SessionInfos, OutHeader, ReadHeaderPos,OutFailReason))
+		{
+			HeaderData Header = OutHeader;
+
+			// Create thread
+			if (!Thread)
+			{
+				FString ThreadName = (Header.SerializedType.ToString());
+				Thread = FRunnableThread::Create(Runnable, *ThreadName);
+				UE_LOGF(MovieSceneSerialization, Display, "Created Thread For Type: %ls", *ThreadName);
+			}
+
+			// Add a command to populate the dictionaries/maps from this existing file first
+			FContextAndCommands<HeaderData,FrameData>& Context = Runnable->GetContext(Guid);
+			Context.AddReadCommand([ReadHeaderPos](FMovieSceneSerializerContext& InContext)
+			{
+				InContext.FrameIdToFrameOffset.Reset();
+				InContext.MinFrameId = 0xffffffffffffffff;
+				InContext.MaxFrameId = 0;
+
+				InContext.ArReader.Seek(ReadHeaderPos);
+
+				while (true)
+				{
+					int64 FrameHeaderPos = InContext.ArReader.Tell();
+
+					FMovieSceneSerializedFrameHeader Header;
+					InContext.ArReader << Header;
+
+					InContext.FrameIdToFrameOffset.Emplace(Header.FrameId, FrameHeaderPos);
+					InContext.MinFrameId = FMath::Min(InContext.MinFrameId, Header.FrameId);
+					InContext.MaxFrameId = FMath::Max(InContext.MaxFrameId, Header.FrameId);
+
+					if (Header.NextFrameOffset != MovieSceneSerializationNamespace::InvalidOffset && Header.NextFrameOffset < InContext.ArReader.TotalSize())
+					{
+						InContext.ArReader.Seek(Header.NextFrameOffset);
+					}
+					else
+					{
+						break;
+					}
+				}
+
+				InContext.ArReader.Seek(ReadHeaderPos);
+			});
+
+			return true;
+		}
+		else
+		{
+			/*
+			Runnable->Stop();
+			delete Runnable;
+			Runnable = nullptr;
+			*/
+			return false;
+		}
+	}
+	else
+	{
+		OutFailReason = NSLOCTEXT("MovieSceneSerializer", "FileOpenFailedToCreateArchive", "Cannot open file.\nFailed to create archive.");
+
+		UE_LOGF(MovieSceneSerialization, Error, "Cannot open Movie Scene Serialization cache %ls. Failed to create archive.", *AbsoluteFilePath);
+		return false;
+	}
+}
+template<typename HeaderData, typename FrameData>
+void TMovieSceneSerializer<HeaderData, FrameData>::Close()
+{
+
+	if (Thread && Runnable)
+	{
+		Runnable->Close(Guid);
+		/* MZ TODO FIGURE WHEN AND HOW TO CLOSE AND DELETE 
+		// Stop the thread (will flush)
+		Runnable->Stop();
+		Thread->WaitForCompletion();
+
+		// Cleanup
+		delete Thread;
+		Thread = nullptr;
+
+		delete Runnable;
+		Runnable = nullptr;
+		*/
+	}
+}
+template<typename HeaderData, typename FrameData>
+bool TMovieSceneSerializer<HeaderData, FrameData>::IsOpen() const
+{
+	return Thread && Runnable;
+}
+
+template<typename HeaderData, typename FrameData>
+void TMovieSceneSerializer<HeaderData, FrameData>::GetDataRanges(TFunction<void(uint64 /*InMinFrameId*/, uint64 /*InMaxFrameId*/)> InCompletionCallback)
+{
+	if (Runnable != nullptr && Runnable->bOpenForRead)
+	{
+		FContextAndCommands<HeaderData,FrameData>& Context = Runnable->GetContext(Guid);
+		Context.AddReadCommand([InCompletionCallback](FMovieSceneSerializerContext& InContext)
+		{
+			uint64 MinFrameId = 0, MaxFrameId = 0;
+			MinFrameId = InContext.MinFrameId;
+			MaxFrameId = InContext.MaxFrameId;
+
+			AsyncTask(ENamedThreads::GameThread, [MinFrameId, MaxFrameId, InCompletionCallback]() mutable
+			{
+				InCompletionCallback(MinFrameId, MaxFrameId);
+			});
+		});
+	}
+}
+
+template <typename HeaderData, typename FrameData>
+void TMovieSceneSerializer<HeaderData, FrameData>::ReadFramesAtFrameRange(uint64 InStartFrameId, uint64 InEndFrameId, TFunction<void()> InCompletionCallback)
+{
+	check(InStartFrameId <= InEndFrameId);
+
+	if (Runnable != nullptr && Runnable->bOpenForRead)
+	{
+		FContextAndCommands<HeaderData,FrameData>& Context = Runnable->GetContext(Guid);
+		Context.AddReadCommand([this,InStartFrameId, InEndFrameId, InCompletionCallback](FMovieSceneSerializerContext& InContext)
+		{
+			// First we need to skip to our frame
+			int32 FoundFirstFramePos = Algo::LowerBoundBy(InContext.FrameIdToFrameOffset, InStartFrameId, [](const TPair<uint64, int64>& InItem) { return InItem.Key; }, TLess<>());
+			if (InContext.FrameIdToFrameOffset.IsValidIndex(FoundFirstFramePos))
+			{
+				int64 ReturnPos = InContext.ArReader.Tell();
+
+				TArray<TMovieSceneSerializedFrame<FrameData>> TempFrames;
+
+				InContext.ArReader.Seek(InContext.FrameIdToFrameOffset[FoundFirstFramePos].Value);
+				while (true)
+				{
+					FMovieSceneSerializedFrameHeader Header;
+					InContext.ArReader << Header;
+
+					TMovieSceneSerializedFrame<FrameData>& TempFrame = TempFrames.Emplace_GetRef();
+					InContext.ArReader << TempFrame;
+					if (Header.NextFrameOffset != MovieSceneSerializationNamespace::InvalidOffset && Header.NextFrameOffset < InContext.ArReader.TotalSize() && Header.FrameId <= InEndFrameId)
+					{
+						InContext.ArReader.Seek(Header.NextFrameOffset);
+					}
+					else
+					{
+						break;
+					}
+				}
+
+				InContext.ArReader.Seek(ReturnPos);
+
+				AsyncTask(ENamedThreads::GameThread, [this, InCompletionCallback, TempFrames = MoveTemp(TempFrames)]() mutable
+				{
+					ResultData = MoveTemp(TempFrames);
+					InCompletionCallback();
+				});
+			}
+			else
+			{
+				AsyncTask(ENamedThreads::GameThread, [this,InCompletionCallback]()
+				{
+					TArray<TMovieSceneSerializedFrame<FrameData>> DummyFrames;
+					ResultData = MoveTemp(DummyFrames);
+					InCompletionCallback();
+				});
+			}
+		});
+	}
+}
+
+template<typename HeaderData, typename FrameData>
+void TMovieSceneSerializer<HeaderData, FrameData>::WriteFrameData(uint64 InFrameId, const FrameData &InFrame)
+{
+	if (Runnable != nullptr && Runnable->bOpenForWrite)
+	{
+		FramesWritten += 1;
+		FContextAndCommands<HeaderData,FrameData>& Context = Runnable->GetContext(Guid);
+		Context.AddWriteCommand([InFrameId, InFrame](FMovieSceneSerializerContext& InContext)
+		{
+			FrameData CurrentBufferedWriteFrame = InFrame;
+			// Seek to the end of the archive
+			InContext.ArWriter.Seek(InContext.ArWriter.TotalSize());
+
+			// Find our current offset
+			int64 ThisFramePos = InContext.ArWriter.Tell();
+
+			// Serialize out the next frame offset (currently invalid) and previous frame offset in a new header
+			FMovieSceneSerializedFrameHeader Header(MovieSceneSerializationNamespace::InvalidOffset, InContext.LastFrameWritePos, InFrameId);
+			InContext.ArWriter << Header;
+
+			int64 FrameId = InFrameId;
+			InContext.ArWriter << FrameId;
+			// Serialize the buffered frame's body
+			InContext.ArWriter << CurrentBufferedWriteFrame;
+			// update caches
+			InContext.FrameIdToFrameOffset.Emplace(InFrameId, ThisFramePos);
+			InContext.MinFrameId = FMath::Min(InContext.MinFrameId, InFrameId);
+			InContext.MaxFrameId = FMath::Max(InContext.MaxFrameId, InFrameId);
+
+			// Serialize next frame pos to prev frame
+			if (InContext.LastFrameWritePos != MovieSceneSerializationNamespace::InvalidOffset)
+			{
+				int64 ReturnPos = InContext.ArWriter.Tell();
+				InContext.ArWriter.Seek(InContext.LastFrameWritePos + STRUCT_OFFSET(FMovieSceneSerializedFrameHeader, NextFrameOffset));
+				InContext.ArWriter << ThisFramePos;
+				InContext.ArWriter.Seek(ReturnPos);
+			}
+
+			InContext.LastFrameWritePos = ThisFramePos;
+		});
+	}
+}
+
+
+template<typename HeaderData, typename FrameData>
+void TMovieSceneSerializer<HeaderData, FrameData>::WriteFrameData(uint64 InFrameId, const TArray<FrameData> &InFrames)
+{
+	if (Runnable != nullptr && Runnable->bOpenForWrite)
+	{
+		FramesWritten += InFrames.Num();
+		FContextAndCommands<HeaderData,FrameData>& Context = Runnable->GetContext(Guid);
+		Context.AddWriteCommand([InFrameId, InFrames](FMovieSceneSerializerContext& InContext)
+		{
+			uint64 FrameId = InFrameId;
+
+			for (const FrameData& InFrame : InFrames)
+			{
+				FrameData CurrentBufferedWriteFrame = InFrame;
+				// Seek to the end of the archive
+				InContext.ArWriter.Seek(InContext.ArWriter.TotalSize());
+
+				// Find our current offset
+				int64 ThisFramePos = InContext.ArWriter.Tell();
+
+				// Serialize out the next frame offset (currently invalid) and previous frame offset in a new header
+				FMovieSceneSerializedFrameHeader Header(MovieSceneSerializationNamespace::InvalidOffset, InContext.LastFrameWritePos, FrameId);
+				InContext.ArWriter << Header;
+
+				InContext.ArWriter << FrameId;
+				// Serialize the buffered frame's body
+				InContext.ArWriter << CurrentBufferedWriteFrame;
+				// update caches
+				InContext.FrameIdToFrameOffset.Emplace(FrameId, ThisFramePos);
+				InContext.MinFrameId = FMath::Min(InContext.MinFrameId, FrameId);
+				InContext.MaxFrameId = FMath::Max(InContext.MaxFrameId, FrameId);
+
+				// Serialize next frame pos to prev frame
+				if (InContext.LastFrameWritePos != MovieSceneSerializationNamespace::InvalidOffset)
+				{
+					int64 ReturnPos = InContext.ArWriter.Tell();
+					InContext.ArWriter.Seek(InContext.LastFrameWritePos + STRUCT_OFFSET(FMovieSceneSerializedFrameHeader, NextFrameOffset));
+					InContext.ArWriter << ThisFramePos;
+					InContext.ArWriter.Seek(ReturnPos);
+				}
+				++FrameId;
+				InContext.LastFrameWritePos = ThisFramePos;
+			}
+		});
+	}
+}
diff --git a/Engine/Source/Runtime/SerializedRecorderInterface/Public/Serializers/MovieSceneSerializedType.h b/Engine/Source/Runtime/SerializedRecorderInterface/Public/Serializers/MovieSceneSerializedType.h
new file mode 100644
--- /dev/null
+++ b/Engine/Source/Runtime/SerializedRecorderInterface/Public/Serializers/MovieSceneSerializedType.h
@@ -0,0 +1,30 @@
+// Copyright Epic Games, Inc. All Rights Reserved.
+
+#pragma once
+
+#include "Misc/FrameRate.h"
+#include "UObject/UnrealType.h"
+#include "MovieSceneSectionSerialization.h"
+
+//Though we are embedding the serialized type of the data into the file name they may not match so we 
+//go ahead and read the type from the file by mocking a header that just reads in the type.
+struct FSerializedTypeFileHeader
+{
+	static const int32 cVersion = 1;
+
+	FSerializedTypeFileHeader() : Version(cVersion)
+	{
+	}
+
+
+	friend FArchive& operator<<(FArchive& Ar, FSerializedTypeFileHeader& Header)
+	{
+		Ar << Header.Version;
+		Ar << Header.SerializedType;
+
+		return Ar;
+	}
+
+	int32 Version;
+	FName SerializedType;
+};
diff --git a/Engine/Source/Runtime/SerializedRecorderInterface/Public/Serializers/MovieSceneSpawnSerialization.h b/Engine/Source/Runtime/SerializedRecorderInterface/Public/Serializers/MovieSceneSpawnSerialization.h
new file mode 100644
--- /dev/null
+++ b/Engine/Source/Runtime/SerializedRecorderInterface/Public/Serializers/MovieSceneSpawnSerialization.h
@@ -0,0 +1,69 @@
+// Copyright Epic Games, Inc. All Rights Reserved.
+
+#pragma once
+
+#include "Misc/FrameRate.h"
+#include "UObject/UnrealType.h"
+#include "MovieSceneSectionSerialization.h"
+
+
+struct FSpawnFileHeader
+{
+	static const int32 cVersion = 1;
+
+	FSpawnFileHeader() : Version(cVersion)
+	{
+	}
+
+
+	FSpawnFileHeader(FFrameRate &InTickResolution, const FName& InSerializedType, const FGuid& InGuid)
+		: Version(cVersion)
+		, SerializedType(InSerializedType)
+		, Guid(InGuid)
+		, TickResolution(InTickResolution)
+
+	{};
+
+
+	friend FArchive& operator<<(FArchive& Ar, FSpawnFileHeader& Header)
+	{
+		Ar << Header.Version;
+		Ar << Header.SerializedType;
+		Ar << Header.Guid;
+		Ar << Header.TickResolution.Numerator;
+		Ar << Header.TickResolution.Denominator;
+		return Ar;
+	}
+
+	//DATA
+	int32 Version;
+	FName SerializedType;
+	FGuid Guid;
+	FFrameRate TickResolution;
+};
+
+
+//since we can't use function templates with TFunctions need to use a variant or union.
+struct FSpawnProperty
+{
+	FSpawnProperty() = default;
+	FSpawnProperty( FFrameNumber InKeyTime, bool InVal)
+	{
+		Time = InKeyTime;
+		bVal = InVal;
+	}
+	friend FArchive& operator<<(FArchive& Ar, FSpawnProperty& Property)
+	{
+		Ar << Property.Time;
+		Ar << Property.bVal;
+		return Ar;
+	}
+
+	//DATA
+	FFrameNumber Time;
+	bool bVal;
+};
+
+
+using FSpawnSerializedFrame = TMovieSceneSerializedFrame<FSpawnProperty>;
+using FSpawnSerializer = TMovieSceneSerializer<FSpawnFileHeader, FSpawnProperty>;
\ No newline at end of file
diff --git a/Engine/Source/Runtime/SerializedRecorderInterface/Public/Serializers/MovieSceneTransformSerialization.h b/Engine/Source/Runtime/SerializedRecorderInterface/Public/Serializers/MovieSceneTransformSerialization.h
new file mode 100644
--- /dev/null
+++ b/Engine/Source/Runtime/SerializedRecorderInterface/Public/Serializers/MovieSceneTransformSerialization.h
@@ -0,0 +1,91 @@
+// Copyright Epic Games, Inc. All Rights Reserved.
+
+#pragma once
+
+#include "Misc/FrameRate.h"
+#include "UObject/UnrealType.h"
+#include "MovieSceneSectionSerialization.h"
+
+
+struct FTransformFileHeader
+{
+	static const int32 cVersion = 1;
+
+	FTransformFileHeader() : Version(cVersion)
+	{
+	}
+
+	FTransformFileHeader(FFrameRate &InTickResolution, const FName& InSerializedType, const FGuid& InGuid)
+		: Version(cVersion)
+		, SerializedType(InSerializedType)
+		, Guid(InGuid)
+		, TickResolution(InTickResolution)
+
+	{};
+
+
+	friend FArchive& operator<<(FArchive& Ar, FTransformFileHeader& Header)
+	{
+		Ar << Header.Version;
+		Ar << Header.SerializedType;
+		Ar << Header.Guid;
+		Ar << Header.TickResolution.Numerator;
+		Ar << Header.TickResolution.Denominator;
+
+		return Ar;
+	}
+
+	//DATA
+	int32 Version;
+	FName SerializedType;
+	FGuid Guid;
+	FFrameRate TickResolution;
+
+};
+
+struct FSerializedTransform
+{
+	FSerializedTransform() = default;
+
+	FSerializedTransform(const FTransform& InTransform, FFrameNumber InKeyTime)
+	{
+		Set(InTransform, InKeyTime);
+	}
+	void Set(const FTransform& InTransform, FFrameNumber InKeyTime)
+	{
+		Time = InKeyTime;
+		Values[0] = InTransform.GetTranslation().X;
+		Values[1] = InTransform.GetTranslation().Y;
+		Values[2] = InTransform.GetTranslation().Z;
+
+		FRotator WoundRoation = InTransform.Rotator();
+		Values[3] = WoundRoation.Roll;
+		Values[4] = WoundRoation.Pitch;
+		Values[5] = WoundRoation.Yaw;
+
+		Values[6] = InTransform.GetScale3D().X;
+		Values[7] = InTransform.GetScale3D().Y;
+		Values[8] = InTransform.GetScale3D().Z;
+	}
+
+	friend FArchive& operator<<(FArchive& Ar, FSerializedTransform& Transform)
+	{
+		Ar << Transform.Time;
+		Ar << Transform.Values[0];
+		Ar << Transform.Values[1];
+		Ar << Transform.Values[2];
+		Ar << Transform.Values[3];
+		Ar << Transform.Values[4];
+		Ar << Transform.Values[5];
+		Ar << Transform.Values[6];
+		Ar << Transform.Values[7];
+		Ar << Transform.Values[8];
+		return Ar;
+	}
+
+	FFrameNumber Time;
+	double Values[9];   //location, rotation, scale.
+};
+
+using FTransformSerializedFrame = TMovieSceneSerializedFrame<FSerializedTransform>;
+using FTransformSerializer = TMovieSceneSerializer<FTransformFileHeader, FSerializedTransform>;
diff --git a/Engine/Source/Runtime/SerializedRecorderInterface/SerializedRecorderInterface.Build.cs b/Engine/Source/Runtime/SerializedRecorderInterface/SerializedRecorderInterface.Build.cs
new file mode 100644
--- /dev/null
+++ b/Engine/Source/Runtime/SerializedRecorderInterface/SerializedRecorderInterface.Build.cs
@@ -0,0 +1,18 @@
+// Copyright Epic Games, Inc. All Rights Reserved.
+
+using UnrealBuildTool;
+
+public class SerializedRecorderInterface : ModuleRules
+{
+    public SerializedRecorderInterface(ReadOnlyTargetRules Target) : base(Target)
+	{
+        PublicDependencyModuleNames.AddRange(
+            new string[] {
+                "Core",
+                "CoreUObject",
+                "LiveLinkInterface",
+
+            }
+        );
+	}
+}
